Fix horde issue 416438 move LiveLinkInputDevice to optional plugin.

diff --git a/Engine/Programs/LiveLinkHub/Config/DefaultEngine.ini b/Engine/Programs/LiveLinkHub/Config/DefaultEngine.ini
--- a/Engine/Programs/LiveLinkHub/Config/DefaultEngine.ini
+++ b/Engine/Programs/LiveLinkHub/Config/DefaultEngine.ini
@@ -7,10 +7,10 @@ GameNameShort=LLH
 +ProgramEnabledPlugins="ContentBrowserAssetDataSource"
 +ProgramEnabledPlugins="PropertyAccessEditor"
 +ProgramEnabledPlugins="PythonScriptPlugin"
-+ProgramEnabledPlugins="LiveLinkInputDevice"
 
 +ProgramOptionalPlugins="AppleARKitFaceSupport"
 #+ProgramOptionalPlugins="LiveLinkViconDataStream"
++ProgramOptionalPlugins="LiveLinkInputDevice"
 +ProgramOptionalPlugins="MocopiLiveLink"
 +ProgramOptionalPlugins="OptitrackLiveLink"
 
diff --git a/Engine/Source/Programs/LiveLinkHub/LiveLinkHub.Target.cs b/Engine/Source/Programs/LiveLinkHub/LiveLinkHub.Target.cs
--- a/Engine/Source/Programs/LiveLinkHub/LiveLinkHub.Target.cs
+++ b/Engine/Source/Programs/LiveLinkHub/LiveLinkHub.Target.cs
@@ -26,12 +26,13 @@ public class LiveLinkHubTarget : TargetRules
 		AdditionalPlugins.Add("QuicMessaging");
 		AdditionalPlugins.Add("PropertyAccessEditor");
 		AdditionalPlugins.Add("PythonScriptPlugin");
-		AdditionalPlugins.Add("LiveLinkInputDevice");
+
 		OptionalPlugins.AddRange(new string[]
 		{
 			"AppleARKitFaceSupport",
 			//"LiveLinkViconDataStream",
 			"MocopiLiveLink",
+			"LiveLinkInputDevice",
 			"OptitrackLiveLink"
 		});
 
