LiveLinkHub - Disable third party plugins

diff --git a/Engine/Programs/LiveLinkHub/Config/DefaultEngine.ini b/Engine/Programs/LiveLinkHub/Config/DefaultEngine.ini
--- a/Engine/Programs/LiveLinkHub/Config/DefaultEngine.ini
+++ b/Engine/Programs/LiveLinkHub/Config/DefaultEngine.ini
@@ -17,12 +17,13 @@ GameNameShort=LLH
 
 ; Optional plugins, will not cause an error if disabled.
 +ProgramOptionalPlugins="AppleARKitFaceSupport"
-+ProgramOptionalPlugins="LiveLinkViconDataStream"
-+ProgramOptionalPlugins="Smartsuit"
-+ProgramOptionalPlugins="PoseAILiveLink"
-+ProgramOptionalPlugins="LiveLinkMVNPlugin"
-+ProgramOptionalPlugins="MocopiLiveLink"
-+ProgramOptionalPlugins="OptitrackLiveLink"
+
+;+ProgramOptionalPlugins="LiveLinkViconDataStream"
+;+ProgramOptionalPlugins="Smartsuit"
+;+ProgramOptionalPlugins="PoseAILiveLink"
+;+ProgramOptionalPlugins="LiveLinkMVNPlugin"
+;+ProgramOptionalPlugins="MocopiLiveLink"
+;+ProgramOptionalPlugins="OptitrackLiveLink"
 
 [/Script/MacTargetPlatform.XcodeProjectSettings]
 TemplateMacPlist=(FilePath="/Engine/Source/Programs/LiveLinkHub/Resources/Mac/Info.plist")
diff --git a/Engine/Source/Programs/LiveLinkHub/LiveLinkHub.Target.cs b/Engine/Source/Programs/LiveLinkHub/LiveLinkHub.Target.cs
--- a/Engine/Source/Programs/LiveLinkHub/LiveLinkHub.Target.cs
+++ b/Engine/Source/Programs/LiveLinkHub/LiveLinkHub.Target.cs
@@ -54,19 +54,6 @@ public class LiveLinkHubTarget : TargetRules
 			"AppleARKitFaceSupport"
 		});
 
-		if (bEnableThirdPartyPlugins)
-		{
-			OptionalPlugins.AddRange(new string[]
-			{
-				"LiveLinkMvnPlugin",
-				"LiveLinkViconDataStream",
-				"MocopiLiveLink",
-				"OptitrackLiveLink",
-				"PoseAILiveLink",
-				"Smartsuit"
-			});
-		}
-
 		bCompileAgainstCoreUObject = true;
 		bCompileAgainstEngine = true;
 		bCompileAgainstEditor = true;
