LiveLinkHub - Enable PoseAI, Mocopi and vicon third party plugins

diff --git a/Engine/Programs/LiveLinkHub/Config/DefaultEngine.ini b/Engine/Programs/LiveLinkHub/Config/DefaultEngine.ini
--- a/Engine/Programs/LiveLinkHub/Config/DefaultEngine.ini
+++ b/Engine/Programs/LiveLinkHub/Config/DefaultEngine.ini
@@ -18,11 +18,11 @@ GameNameShort=LLH
 ; Optional plugins, will not cause an error if disabled.
 +ProgramOptionalPlugins="AppleARKitFaceSupport"
 
-;+ProgramOptionalPlugins="LiveLinkViconDataStream"
++ProgramOptionalPlugins="LiveLinkViconDataStream"
 ;+ProgramOptionalPlugins="Smartsuit"
-;+ProgramOptionalPlugins="PoseAILiveLink"
++ProgramOptionalPlugins="PoseAILiveLink"
 ;+ProgramOptionalPlugins="LiveLinkMVNPlugin"
-;+ProgramOptionalPlugins="MocopiLiveLink"
++ProgramOptionalPlugins="MocopiLiveLink"
 ;+ProgramOptionalPlugins="OptitrackLiveLink"
 
 [/Script/MacTargetPlatform.XcodeProjectSettings]
diff --git a/Engine/Source/Programs/LiveLinkHub/LiveLinkHub.Target.cs b/Engine/Source/Programs/LiveLinkHub/LiveLinkHub.Target.cs
--- a/Engine/Source/Programs/LiveLinkHub/LiveLinkHub.Target.cs
+++ b/Engine/Source/Programs/LiveLinkHub/LiveLinkHub.Target.cs
@@ -48,6 +48,17 @@ public class LiveLinkHubTarget : TargetRules
 			"StructUtils",
 			"UdpMessaging"
 		});
+		
+		
+		if (bEnableThirdPartyPlugins)
+		{
+			OptionalPlugins.AddRange(new string[]
+			{
+				"LiveLinkViconDataStream",
+				"MocopiLiveLink",
+				"PoseAILiveLink"
+			});
+		}
 
 		OptionalPlugins.AddRange(new string[]
 		{
