LiveLinkHub - Enable Rokoko plugin

diff --git a/Engine/Programs/LiveLinkHub/Config/DefaultEngine.ini b/Engine/Programs/LiveLinkHub/Config/DefaultEngine.ini
--- a/Engine/Programs/LiveLinkHub/Config/DefaultEngine.ini
+++ b/Engine/Programs/LiveLinkHub/Config/DefaultEngine.ini
@@ -19,7 +19,7 @@ GameNameShort=LLH
 +ProgramOptionalPlugins="AppleARKitFaceSupport"
 
 +ProgramOptionalPlugins="LiveLinkViconDataStream"
-;+ProgramOptionalPlugins="Smartsuit"
++ProgramOptionalPlugins="Smartsuit"
 +ProgramOptionalPlugins="PoseAILiveLink"
 ;+ProgramOptionalPlugins="LiveLinkMVNPlugin"
 +ProgramOptionalPlugins="MocopiLiveLink"
diff --git a/Engine/Source/Programs/LiveLinkHub/LiveLinkHub.Target.cs b/Engine/Source/Programs/LiveLinkHub/LiveLinkHub.Target.cs
--- a/Engine/Source/Programs/LiveLinkHub/LiveLinkHub.Target.cs
+++ b/Engine/Source/Programs/LiveLinkHub/LiveLinkHub.Target.cs
@@ -56,7 +56,8 @@ public class LiveLinkHubTarget : TargetRules
 			{
 				"LiveLinkViconDataStream",
 				"MocopiLiveLink",
-				"PoseAILiveLink"
+				"PoseAILiveLink",
+				"Smartsuit"
 			});
 		}
 
